Fix VirtualCameraPixelStreamingPresetEditor plugin busting size limit of 260 in editor builds by using ShortName in build.cs

diff --git a/Engine/Plugins/Experimental/VirtualCameraPixelStreamingPreset/Source/VirtualCameraPixelStreamingPresetEditor/VirtualCameraPixelStreamingPresetEditor.Build.cs b/Engine/Plugins/Experimental/VirtualCameraPixelStreamingPreset/Source/VirtualCameraPixelStreamingPresetEditor/VirtualCameraPixelStreamingPresetEditor.Build.cs
--- a/Engine/Plugins/Experimental/VirtualCameraPixelStreamingPreset/Source/VirtualCameraPixelStreamingPresetEditor/VirtualCameraPixelStreamingPresetEditor.Build.cs
+++ b/Engine/Plugins/Experimental/VirtualCameraPixelStreamingPreset/Source/VirtualCameraPixelStreamingPresetEditor/VirtualCameraPixelStreamingPresetEditor.Build.cs
@@ -6,6 +6,7 @@ public class VirtualCameraPixelStreamingPresetEditor : ModuleRules
 {
 	public VirtualCameraPixelStreamingPresetEditor(ReadOnlyTargetRules Target) : base(Target)
 	{
+		ShortName = "VCamPSPresetEditor";
 		DefaultBuildSettings = BuildSettingsVersion.V2;
 
 		PrivateDependencyModuleNames.AddRange(
